TargetDeviceServices - Add plugin to be able to work with information stored in the TargetDeviceServices module via blueprints and Python.

diff --git a/Engine/Plugins/Experimental/TargetDeviceServicesScripting/Source/TargetDeviceServicesScripting/Private/TargetDeviceServicesBPFunctionLibrary.cpp b/Engine/Plugins/Experimental/TargetDeviceServicesScripting/Source/TargetDeviceServicesScripting/Private/TargetDeviceServicesBPFunctionLibrary.cpp
new file mode 100644
--- /dev/null
+++ b/Engine/Plugins/Experimental/TargetDeviceServicesScripting/Source/TargetDeviceServicesScripting/Private/TargetDeviceServicesBPFunctionLibrary.cpp
@@ -0,0 +1,77 @@
+// Copyright Epic Games, Inc. All Rights Reserved.
+
+#include "TargetDeviceServicesBPFunctionLibrary.h"
+#include "Modules/ModuleManager.h"
+#include "ITargetDeviceServicesModule.h"
+#include "ITargetDeviceProxyManager.h"
+#include "ITargetDeviceService.h"
+#include "ITargetDeviceProxy.h"
+
+UTargetDeviceServicesBPFunctionLibrary::UTargetDeviceServicesBPFunctionLibrary(const FObjectInitializer& ObjectInitializer)
+	: Super(ObjectInitializer)
+{
+}
+
+TMap<FString, FDeviceSnapshots> UTargetDeviceServicesBPFunctionLibrary::GetDeviceSnapshots()
+{
+	ITargetDeviceServicesModule& TargetDeviceServicesModule = FModuleManager::GetModuleChecked<ITargetDeviceServicesModule>("TargetDeviceServices");
+	TSharedRef<ITargetDeviceProxyManager> DeviceProxyManager = TargetDeviceServicesModule.GetDeviceProxyManager();
+	TSharedRef<ITargetDeviceServiceManager> DeviceServiceManager = TargetDeviceServicesModule.GetDeviceServiceManager();
+
+	TArray<ITargetDeviceServicePtr> DeviceServices;
+	const int32 DeviceServicesCount = DeviceServiceManager->GetServices(DeviceServices);
+
+	TMap<FString, FDeviceSnapshots> Snapshots;
+	for (int32 CurrentDeviceServiceIndex = 0; CurrentDeviceServiceIndex < DeviceServicesCount; ++CurrentDeviceServiceIndex)
+	{
+		const ITargetDeviceServicePtr& DeviceService = DeviceServices[CurrentDeviceServiceIndex];
+
+		if (DeviceService.IsValid())
+		{
+			// Fetching shared pointer to default device by passing into the ITargetDeviceService::GetDevice method Name_None
+			ITargetDevicePtr DefaultDevice = DeviceService->GetDevice(NAME_None);
+
+			if (DefaultDevice.IsValid())
+			{
+				const FString DefaultDeviceType = TargetDeviceTypes::ToString(DefaultDevice->GetDeviceType());
+				FString DefaultDeviceId = DefaultDevice->GetId().ToString();
+				TSharedPtr<ITargetDeviceProxy> DefaultDeviceProxy = 
+					DeviceProxyManager->FindProxyDeviceForTargetDevice(DefaultDeviceId);
+
+				FDeviceSnapshot DeviceSnapshot{
+					// Fills the Name field
+					DeviceService->GetDeviceName(),
+					// Fills the HostName field
+					DefaultDeviceProxy.IsValid() ? DefaultDeviceProxy->GetHostName() : DefaultDevice->GetName(),
+					// Fills the DeviceType field
+					DefaultDeviceType,
+					// fills the ModelId field
+					DefaultDevice->GetModelId(),
+					// fills the DeviceConnectionType;
+					DefaultDeviceProxy.IsValid() ? DefaultDeviceProxy->GetConnectionType() : TEXT(""),
+					// fills the DeviceId field
+					MoveTemp(DefaultDeviceId),
+					// fills the OperatingSystem field
+					DeviceService->GetDevicePlatformDisplayName(),
+					// fills the IsConnected field
+					DefaultDevice->IsConnected()
+				};
+
+				FDeviceSnapshots* DeviceSnapshotsOfCurrentDeviceTypePtr = Snapshots.Find(DefaultDeviceType);
+				if (nullptr == DeviceSnapshotsOfCurrentDeviceTypePtr)
+				{
+					//Initializer list constructor call
+					TArray<FDeviceSnapshot> Entries{MoveTemp(DeviceSnapshot)};
+					Snapshots.Emplace(DefaultDeviceType, FDeviceSnapshots{MoveTemp(Entries)});
+				}
+				else
+				{
+					DeviceSnapshotsOfCurrentDeviceTypePtr->Entries.Emplace(MoveTemp(DeviceSnapshot));
+				}
+			}
+		}
+	}
+
+	return MoveTemp(Snapshots);
+}
+
diff --git a/Engine/Plugins/Experimental/TargetDeviceServicesScripting/Source/TargetDeviceServicesScripting/Private/TargetDeviceServicesScriptingModule.cpp b/Engine/Plugins/Experimental/TargetDeviceServicesScripting/Source/TargetDeviceServicesScripting/Private/TargetDeviceServicesScriptingModule.cpp
new file mode 100644
--- /dev/null
+++ b/Engine/Plugins/Experimental/TargetDeviceServicesScripting/Source/TargetDeviceServicesScripting/Private/TargetDeviceServicesScriptingModule.cpp
@@ -0,0 +1,8 @@
+// Copyright Epic Games, Inc. All Rights Reserved.
+
+#include "Modules/ModuleManager.h"
+
+IMPLEMENT_MODULE(FDefaultModuleImpl, TargetDeviceServicesScripting);
+
+
+
diff --git a/Engine/Plugins/Experimental/TargetDeviceServicesScripting/Source/TargetDeviceServicesScripting/Public/TargetDeviceServicesBPFunctionLibrary.h b/Engine/Plugins/Experimental/TargetDeviceServicesScripting/Source/TargetDeviceServicesScripting/Public/TargetDeviceServicesBPFunctionLibrary.h
new file mode 100644
--- /dev/null
+++ b/Engine/Plugins/Experimental/TargetDeviceServicesScripting/Source/TargetDeviceServicesScripting/Public/TargetDeviceServicesBPFunctionLibrary.h
@@ -0,0 +1,88 @@
+// Copyright Epic Games, Inc. All Rights Reserved.
+#pragma once
+
+#include "Kismet/BlueprintFunctionLibrary.h"
+#include "CoreMinimal.h"
+#include "Containers/Map.h"
+#include "TargetDeviceServicesBPFunctionLibrary.generated.h"
+
+/**
+* The struct is designed to store device information
+*/
+USTRUCT(BlueprintType)
+struct FDeviceSnapshot
+{
+	GENERATED_BODY()
+
+public:
+
+	/** Stores device's name. */
+	UPROPERTY(BlueprintReadOnly, Category = "TargetDeviceServicesScripting")
+	FString Name;
+
+	/** Stores device's hostname. */
+	UPROPERTY(BlueprintReadOnly, Category = "TargetDeviceServicesScripting")
+	FString HostName;
+
+	/** Stores device's type. */
+	UPROPERTY(BlueprintReadOnly, Category = "TargetDeviceServicesScripting")
+	FString DeviceType;
+
+	/** Stores device's model identifier. */
+	UPROPERTY(BlueprintReadOnly, Category = "TargetDeviceServicesScripting")
+	FString ModelId;
+
+	/** Stores device's connection type. */
+	UPROPERTY(BlueprintReadOnly, Category = "TargetDeviceServicesScripting")
+	FString DeviceConnectionType;
+
+	/** Stores device's identifier. */
+	UPROPERTY(BlueprintReadOnly, Category = "TargetDeviceServicesScripting")
+	FString DeviceId;
+	
+	/** Stores device's operating system name. */
+	UPROPERTY(BlueprintReadOnly, Category = "TargetDeviceServicesScripting")
+	FString OperatingSystem;
+	
+	/**
+	* Stores device's flag that is used to detect whether device
+	* is connected (true) or disconnected (false). 
+	*/
+	UPROPERTY(BlueprintReadOnly, Category = "TargetDeviceServicesScripting")
+	bool IsConnected = false;
+};
+
+/**
+* The struct is a container class that stores instances of FDeviceSnapshot instances.
+* It is a helper class to be used in blueprints as a value of TMap container.
+*/
+USTRUCT(BlueprintType)
+struct FDeviceSnapshots
+{
+	GENERATED_BODY()
+
+public:
+
+	/** Stores array of device snapshots. */
+	UPROPERTY(BlueprintReadOnly, Category = "TargetDeviceServicesScripting")
+	TArray<FDeviceSnapshot> Entries;
+};
+
+/**
+* The class declares a set of functions to be exposed to blueprints.
+*/
+UCLASS(meta = (ScriptName = "TargetDeviceServices"))
+class UTargetDeviceServicesBPFunctionLibrary : public UBlueprintFunctionLibrary
+{
+	GENERATED_UCLASS_BODY()
+	
+public:
+	/**
+	* Fetches snapshots of devices that are available in the network.
+	* 
+	* @return A dictionary of devices' informational snapshots that are grouped by device type (device type string is used as a key).
+	*/
+	UFUNCTION(BlueprintCallable, Category = "TargetDeviceServicesScripting")
+	static TARGETDEVICESERVICESSCRIPTING_API TMap<FString, FDeviceSnapshots> GetDeviceSnapshots();
+};
+
diff --git a/Engine/Plugins/Experimental/TargetDeviceServicesScripting/Source/TargetDeviceServicesScripting/TargetDeviceServicesScripting.Build.cs b/Engine/Plugins/Experimental/TargetDeviceServicesScripting/Source/TargetDeviceServicesScripting/TargetDeviceServicesScripting.Build.cs
new file mode 100644
--- /dev/null
+++ b/Engine/Plugins/Experimental/TargetDeviceServicesScripting/Source/TargetDeviceServicesScripting/TargetDeviceServicesScripting.Build.cs
@@ -0,0 +1,24 @@
+// Copyright Epic Games, Inc. All Rights Reserved.
+
+using UnrealBuildTool;
+
+namespace UnrealBuildTool.Rules
+{
+	public class TargetDeviceServicesScripting : ModuleRules
+	{
+		public TargetDeviceServicesScripting(ReadOnlyTargetRules Target) : base(Target)
+		{
+			PublicDependencyModuleNames.AddRange(
+				new string[] {
+					"Core",
+					"Engine",
+				});
+
+			PrivateDependencyModuleNames.AddRange(
+				new string[] {
+					"CoreUObject",
+					"TargetDeviceServices",
+				});
+		}
+	}
+}
diff --git a/Engine/Plugins/Experimental/TargetDeviceServicesScripting/TargetDeviceServicesScripting.uplugin b/Engine/Plugins/Experimental/TargetDeviceServicesScripting/TargetDeviceServicesScripting.uplugin
new file mode 100644
--- /dev/null
+++ b/Engine/Plugins/Experimental/TargetDeviceServicesScripting/TargetDeviceServicesScripting.uplugin
@@ -0,0 +1,22 @@
+{
+	"FileVersion": 1,
+	"Version": 1,
+	"VersionName": "1.0",
+	"FriendlyName": "TargetDeviceServices scripting library",
+	"Description": "Set of blueprint functions that enables working with TargetDeviceServices module via scripting",
+	"Category": "Scripting",
+	"CreatedBy": "Epic Games, Inc.",
+	"CreatedByURL": "https://epicgames.com",
+	"EnabledByDefault": false,
+	"CanContainContent": false,
+	"IsBetaVersion": true,
+	"Installed": false,
+	"IsExperimentalVersion": true,
+	"Modules": [
+		{
+			"Name": "TargetDeviceServicesScripting",
+			"Type": "Editor",
+			"LoadingPhase": "Default"
+		}
+	]
+}
\ No newline at end of file
