TargetDeviceServices - Fix long names in intermediate file names.

diff --git a/Engine/Plugins/Experimental/TargetDeviceServicesScripting/Source/TargetDeviceServicesScripting/TargetDeviceServicesScripting.Build.cs b/Engine/Plugins/Experimental/TargetDeviceServicesScripting/Source/TargetDeviceServicesScripting/TargetDeviceServicesScripting.Build.cs
--- a/Engine/Plugins/Experimental/TargetDeviceServicesScripting/Source/TargetDeviceServicesScripting/TargetDeviceServicesScripting.Build.cs
+++ b/Engine/Plugins/Experimental/TargetDeviceServicesScripting/Source/TargetDeviceServicesScripting/TargetDeviceServicesScripting.Build.cs
@@ -19,6 +19,8 @@ namespace UnrealBuildTool.Rules
 					"CoreUObject",
 					"TargetDeviceServices",
 				});
+
+			ShortName = "TDSScripting";
 		}
 	}
 }
