Updated uplugin descriptor files marked as both Experimental and Beta. Plugins with both flags in the Developer, Runtime, and Editor folders will be marked as Beta, while plugins in the Experimental folder will be marked Experimental.

This change resolves some ambiguity around beta plugins that had both flags enabled, after a change in CL#36631664 correctly made the experimental flag take precedence in the plugin browser. This resulted in certain plugins appearing to revert from Beta to Experimental, so we're restoring the previous status for plugins that were intended to be in Beta.

diff --git a/Engine/Plugins/Animation/ControlRigSpline/ControlRigSpline.uplugin b/Engine/Plugins/Animation/ControlRigSpline/ControlRigSpline.uplugin
--- a/Engine/Plugins/Animation/ControlRigSpline/ControlRigSpline.uplugin
+++ b/Engine/Plugins/Animation/ControlRigSpline/ControlRigSpline.uplugin
@@ -13,7 +13,7 @@
 	"EnabledByDefault": true,
 	"CanContainContent": true,
 	"IsBetaVersion": true,
-	"IsExperimentalVersion":  true,
+	"IsExperimentalVersion":  false,
 	"Installed": false,
 	"Plugins": [
 		{
diff --git a/Engine/Plugins/Developer/FieldNotificationTrace/FieldNotificationTrace.uplugin b/Engine/Plugins/Developer/FieldNotificationTrace/FieldNotificationTrace.uplugin
--- a/Engine/Plugins/Developer/FieldNotificationTrace/FieldNotificationTrace.uplugin
+++ b/Engine/Plugins/Developer/FieldNotificationTrace/FieldNotificationTrace.uplugin
@@ -11,7 +11,7 @@
 	"EnabledByDefault": false,
 	"CanContainContent": false,
 	"IsBetaVersion": true,
-	"IsExperimentalVersion": true,
+	"IsExperimentalVersion": false,
 	"Installed": false,
 
 	"Modules": [
diff --git a/Engine/Plugins/Editor/ModelViewViewModelPreview/ModelViewViewModelPreview.uplugin b/Engine/Plugins/Editor/ModelViewViewModelPreview/ModelViewViewModelPreview.uplugin
--- a/Engine/Plugins/Editor/ModelViewViewModelPreview/ModelViewViewModelPreview.uplugin
+++ b/Engine/Plugins/Editor/ModelViewViewModelPreview/ModelViewViewModelPreview.uplugin
@@ -1,4 +1,4 @@
-﻿{
+{
 	"FileVersion": 3,
 	"FriendlyName": "UMG Viewmodel for UMG Preview",
 	"Version": 1,
@@ -10,7 +10,7 @@
 	"EnabledByDefault": false,
 	"CanContainContent": false,
 	"IsBetaVersion": true,
-	"IsExperimentalVersion": true,
+	"IsExperimentalVersion": false,
 	"Installed": false,
 	"Modules": [
 		{
diff --git a/Engine/Plugins/Experimental/ChaosCaching/ChaosCaching.uplugin b/Engine/Plugins/Experimental/ChaosCaching/ChaosCaching.uplugin
--- a/Engine/Plugins/Experimental/ChaosCaching/ChaosCaching.uplugin
+++ b/Engine/Plugins/Experimental/ChaosCaching/ChaosCaching.uplugin
@@ -12,7 +12,7 @@
 	"SupportURL": "",
 	"EnabledByDefault": false,
 	"CanContainContent": true,
-	"IsBetaVersion": true,
+	"IsBetaVersion": false,
 	"IsExperimentalVersion": true,
 	"Installed": false,
 	"Modules": [
diff --git a/Engine/Plugins/Experimental/ChaosModularVehicle/ChaosModularVehicle.uplugin b/Engine/Plugins/Experimental/ChaosModularVehicle/ChaosModularVehicle.uplugin
--- a/Engine/Plugins/Experimental/ChaosModularVehicle/ChaosModularVehicle.uplugin
+++ b/Engine/Plugins/Experimental/ChaosModularVehicle/ChaosModularVehicle.uplugin
@@ -11,7 +11,7 @@
 	"MarketplaceURL": "",
 	"SupportURL": "",
 	"CanContainContent": true,
-	"IsBetaVersion": true,
+	"IsBetaVersion": false,
 	"IsExperimentalVersion": true,
 	"Installed": false,
 	"Modules": [
diff --git a/Engine/Plugins/Experimental/SurfaceEffects/SurfaceEffects.uplugin b/Engine/Plugins/Experimental/SurfaceEffects/SurfaceEffects.uplugin
--- a/Engine/Plugins/Experimental/SurfaceEffects/SurfaceEffects.uplugin
+++ b/Engine/Plugins/Experimental/SurfaceEffects/SurfaceEffects.uplugin
@@ -10,7 +10,7 @@
 	"DocsURL": "",
 	"MarketplaceURL": "",
 	"CanContainContent": false,
-	"IsBetaVersion": true,
+	"IsBetaVersion": false,
 	"IsExperimentalVersion": true,
 	"Installed": false,
 	"Modules": [
diff --git a/Engine/Plugins/Experimental/TargetDeviceServicesScripting/TargetDeviceServicesScripting.uplugin b/Engine/Plugins/Experimental/TargetDeviceServicesScripting/TargetDeviceServicesScripting.uplugin
--- a/Engine/Plugins/Experimental/TargetDeviceServicesScripting/TargetDeviceServicesScripting.uplugin
+++ b/Engine/Plugins/Experimental/TargetDeviceServicesScripting/TargetDeviceServicesScripting.uplugin
@@ -9,9 +9,9 @@
 	"CreatedByURL": "https://epicgames.com",
 	"EnabledByDefault": false,
 	"CanContainContent": false,
-	"IsBetaVersion": true,
-	"Installed": false,
+	"IsBetaVersion": false,
 	"IsExperimentalVersion": true,
+	"Installed": false,
 	"Modules": [
 		{
 			"Name": "TargetDeviceServicesScripting",
diff --git a/Engine/Plugins/Experimental/Web/WebAPI/WebAPI.uplugin b/Engine/Plugins/Experimental/Web/WebAPI/WebAPI.uplugin
--- a/Engine/Plugins/Experimental/Web/WebAPI/WebAPI.uplugin
+++ b/Engine/Plugins/Experimental/Web/WebAPI/WebAPI.uplugin
@@ -12,7 +12,7 @@
 	"SupportURL": "",
 	"EnabledByDefault" : false,
 	"CanContainContent": true,
-	"IsBetaVersion": true,
+	"IsBetaVersion": false,
 	"IsExperimentalVersion": true,
 	"Installed": false,
 	"Modules": [
diff --git a/Engine/Plugins/Experimental/Workspace/Workspace.uplugin b/Engine/Plugins/Experimental/Workspace/Workspace.uplugin
--- a/Engine/Plugins/Experimental/Workspace/Workspace.uplugin
+++ b/Engine/Plugins/Experimental/Workspace/Workspace.uplugin
@@ -12,7 +12,7 @@
 	"SupportURL" : "",
 	"EnabledByDefault" : false,
 	"CanContainContent" : false,
-	"IsBetaVersion": true,
+	"IsBetaVersion": false,
 	"IsExperimentalVersion": true,
 	"Installed" : false,
 	"Plugins" :
diff --git a/Engine/Plugins/FX/NiagaraMRQ/NiagaraMRQ.uplugin b/Engine/Plugins/FX/NiagaraMRQ/NiagaraMRQ.uplugin
--- a/Engine/Plugins/FX/NiagaraMRQ/NiagaraMRQ.uplugin
+++ b/Engine/Plugins/FX/NiagaraMRQ/NiagaraMRQ.uplugin
@@ -13,7 +13,7 @@
 	"EnabledByDefault": false,
 	"CanContainContent": true,
 	"IsBetaVersion": true,
-	"IsExperimentalVersion": true,
+	"IsExperimentalVersion": false,
 	"Installed": false,
 	"Modules": [
 		{
diff --git a/Engine/Plugins/Runtime/ModelViewViewModel/ModelViewViewModel.uplugin b/Engine/Plugins/Runtime/ModelViewViewModel/ModelViewViewModel.uplugin
--- a/Engine/Plugins/Runtime/ModelViewViewModel/ModelViewViewModel.uplugin
+++ b/Engine/Plugins/Runtime/ModelViewViewModel/ModelViewViewModel.uplugin
@@ -11,7 +11,7 @@
 	"EnabledByDefault": false,
 	"CanContainContent": false,
 	"IsBetaVersion": true,
-	"IsExperimentalVersion": true,
+	"IsExperimentalVersion": false,
 	"Installed": false,
 
 	"Modules": [
diff --git a/Engine/Plugins/Runtime/Soundscape/Soundscape.uplugin b/Engine/Plugins/Runtime/Soundscape/Soundscape.uplugin
--- a/Engine/Plugins/Runtime/Soundscape/Soundscape.uplugin
+++ b/Engine/Plugins/Runtime/Soundscape/Soundscape.uplugin
@@ -13,7 +13,7 @@
 	"EnabledByDefault": false,
 	"CanContainContent": true,
 	"IsBetaVersion": true,
-	"IsExperimentalVersion": true,
+	"IsExperimentalVersion": false,
 	"Installed": false,
 	"Modules": [
 		{
diff --git a/Engine/Plugins/Runtime/XR/XRVisualization/XRVisualization.uplugin b/Engine/Plugins/Runtime/XR/XRVisualization/XRVisualization.uplugin
--- a/Engine/Plugins/Runtime/XR/XRVisualization/XRVisualization.uplugin
+++ b/Engine/Plugins/Runtime/XR/XRVisualization/XRVisualization.uplugin
@@ -12,8 +12,8 @@
 	"SupportURL": "",
 	"EnabledByDefault": false,
 	"CanContainContent": false,
-	"IsExperimentalVersion": true,
 	"IsBetaVersion": true,
+	"IsExperimentalVersion": false,
 	"Installed": false,
 	"Modules": [
 		{
diff --git a/Engine/Plugins/Tests/InterchangeTests/InterchangeTests.uplugin b/Engine/Plugins/Tests/InterchangeTests/InterchangeTests.uplugin
--- a/Engine/Plugins/Tests/InterchangeTests/InterchangeTests.uplugin
+++ b/Engine/Plugins/Tests/InterchangeTests/InterchangeTests.uplugin
@@ -13,8 +13,8 @@
 	"EnabledByDefault": true,
 	"CanContainContent": false,
 	"IsBetaVersion": true,
+	"IsExperimentalVersion": false,
 	"Installed": false,
-	"IsExperimentalVersion": true,
 	"Modules": [
 		{
 			"Name": "InterchangeTests",
